Adding missing uplugin file to GoogleARCoreServices

diff --git a/Engine/Plugins/Runtime/GoogleARCoreServices/GoogleARCoreServices.uplugin b/Engine/Plugins/Runtime/GoogleARCoreServices/GoogleARCoreServices.uplugin
new file mode 100644
--- /dev/null
+++ b/Engine/Plugins/Runtime/GoogleARCoreServices/GoogleARCoreServices.uplugin
@@ -0,0 +1,24 @@
+{
+    "FileVersion" : 3,
+
+    "Version" : 5,
+    "VersionName" : "1.2",
+    "FriendlyName" : "Google ARCore Services",
+    "Description" : "Provide functionality in Google cross-platform ARCore services.",
+    "Category" : "Augmented Reality",
+    "CreatedBy" : "Google",
+    "CreatedByURL" : "https://developers.google.com/ar/",
+    "CanContainContent": true,
+    "IsBetaVersion": false,
+    "EnabledByDefault": false,
+    "Modules" :
+    [
+        {
+            "Name" : "GoogleARCoreServices",
+            "Type" : "Runtime",
+            "WhitelistPlatforms": [ "Win64", "Mac", "Android", "Linux", "iOS"]
+        },
+    ],
+    "Plugins": [
+    ]
+}
