[GoogleARCoreServices]
* Fixed android compile error

#rn none

diff --git a/Engine/Plugins/Runtime/AR/Google/GoogleARCoreServices/Source/GoogleARCoreServices/Private/GoogleARCoreCloudARPinManager.cpp b/Engine/Plugins/Runtime/AR/Google/GoogleARCoreServices/Source/GoogleARCoreServices/Private/GoogleARCoreCloudARPinManager.cpp
--- a/Engine/Plugins/Runtime/AR/Google/GoogleARCoreServices/Source/GoogleARCoreServices/Private/GoogleARCoreCloudARPinManager.cpp
+++ b/Engine/Plugins/Runtime/AR/Google/GoogleARCoreServices/Source/GoogleARCoreServices/Private/GoogleARCoreCloudARPinManager.cpp
@@ -7,6 +7,10 @@
 #include "arcore_ios_c_api.h"
 #endif
 
+#if ARCORE_SERVICE_SUPPORTED_PLATFORM
+#include "IXRTrackingSystem.h"
+#endif
+
 namespace {
 #if ARCORE_SERVICE_SUPPORTED_PLATFORM
 	EARPinCloudTaskResult ToCloudTaskResult(ArStatus Status)
