Adding missing include.

diff --git a/Engine/Plugins/Runtime/AR/Google/GoogleARCoreServices/Source/GoogleARCoreServices/Private/GoogleARCoreCloudARPinManager.cpp b/Engine/Plugins/Runtime/AR/Google/GoogleARCoreServices/Source/GoogleARCoreServices/Private/GoogleARCoreCloudARPinManager.cpp
--- a/Engine/Plugins/Runtime/AR/Google/GoogleARCoreServices/Source/GoogleARCoreServices/Private/GoogleARCoreCloudARPinManager.cpp
+++ b/Engine/Plugins/Runtime/AR/Google/GoogleARCoreServices/Source/GoogleARCoreServices/Private/GoogleARCoreCloudARPinManager.cpp
@@ -1,6 +1,8 @@
 // Copyright Epic Games, Inc. All Rights Reserved.
 
 #include "GoogleARCoreCloudARPinManager.h"
+#include "IXRTrackingSystem.h"
+
 #if PLATFORM_ANDROID
 #include "arcore_c_api.h"
 #elif PLATFORM_IOS
