UE-183123//UE5/Main - Compile EngineTest Win64 - Warning: Plugin 'GoogleARCoreServices' does not list plugin 'XRBase' as a dependency, but module 'GoogleARCoreServices' depends on 'XRBase'.
-Added XRBase as a plugin dependency.

diff --git a/Engine/Plugins/Runtime/AR/Google/GoogleARCoreServices/GoogleARCoreServices.uplugin b/Engine/Plugins/Runtime/AR/Google/GoogleARCoreServices/GoogleARCoreServices.uplugin
--- a/Engine/Plugins/Runtime/AR/Google/GoogleARCoreServices/GoogleARCoreServices.uplugin
+++ b/Engine/Plugins/Runtime/AR/Google/GoogleARCoreServices/GoogleARCoreServices.uplugin
@@ -19,6 +19,10 @@
             "PlatformAllowList": [ "Win64", "Mac", "Android", "Linux", "IOS"]
         },
     ],
-    "Plugins": [
-    ]
+	"Plugins": [
+		{
+			"Name": "XRBase",
+			"Enabled": true
+		}
+	]
 }
