Add another missing FGCObject::AddReferencedObjects

#rnx
#horde 300570

diff --git a/Engine/Plugins/Runtime/AR/Google/GoogleARCoreServices/Source/GoogleARCoreServices/Private/GoogleARCoreCloudARPinManager.h b/Engine/Plugins/Runtime/AR/Google/GoogleARCoreServices/Source/GoogleARCoreServices/Private/GoogleARCoreCloudARPinManager.h
--- a/Engine/Plugins/Runtime/AR/Google/GoogleARCoreServices/Source/GoogleARCoreServices/Private/GoogleARCoreCloudARPinManager.h
+++ b/Engine/Plugins/Runtime/AR/Google/GoogleARCoreServices/Source/GoogleARCoreServices/Private/GoogleARCoreCloudARPinManager.h
@@ -58,6 +58,11 @@ protected:
 		}
 	}
 
+	virtual FString GetReferencerName() const override
+	{
+		return "GoogleARCoreCloudARPinManager";
+	}
+
 	virtual void OnAlignmentTransformUpdated(const FTransform& NewAlignmentTransform)
 	{
 		for (UCloudARPin* ARPin : AllCloudARPins)
