GoogleARCoreServices ios compile fixes
-#ifdef in the old code for ios for now, we need to update it's api as well but are not ready to do that right now.

diff --git a/Engine/Plugins/Runtime/AR/Google/GoogleARCoreServices/Source/GoogleARCoreServices/Private/GoogleARCoreCloudARPinManager.cpp b/Engine/Plugins/Runtime/AR/Google/GoogleARCoreServices/Source/GoogleARCoreServices/Private/GoogleARCoreCloudARPinManager.cpp
--- a/Engine/Plugins/Runtime/AR/Google/GoogleARCoreServices/Source/GoogleARCoreServices/Private/GoogleARCoreCloudARPinManager.cpp
+++ b/Engine/Plugins/Runtime/AR/Google/GoogleARCoreServices/Source/GoogleARCoreServices/Private/GoogleARCoreCloudARPinManager.cpp
@@ -50,6 +50,10 @@ namespace {
 		// an ArCloudAnchorState, so this function does not ever return those values.
 		case AR_CLOUD_ANCHOR_STATE_NONE:
 			return ECloudARPinCloudState::NotHosted;
+#if ARCORE_USE_OLD_CLOUD_ANCHOR_ASYNC
+		case AR_CLOUD_ANCHOR_STATE_TASK_IN_PROGRESS:
+			return ECloudARPinCloudState::InProgress;
+#endif
 		case AR_CLOUD_ANCHOR_STATE_SUCCESS:
 			return ECloudARPinCloudState::Success;
 		case AR_CLOUD_ANCHOR_STATE_ERROR_INTERNAL:
@@ -281,6 +285,18 @@ FGoogleARCoreCloudARPinManager* FGoogleARCoreCloudARPinManager::CreateCloudARPin
 }
 
 #if ARCORE_SERVICE_SUPPORTED_PLATFORM
+#if ARCORE_USE_OLD_CLOUD_ANCHOR_ASYNC
+void UpdateCloudARPin(UCloudARPin* CloudARPin, ArSession* SessionHandle, ArAnchor* AnchorHandle)
+{
+	ArCloudAnchorState NewCloudState = AR_CLOUD_ANCHOR_STATE_NONE;
+	ArAnchor_getCloudAnchorState(SessionHandle, AnchorHandle, &NewCloudState);
+	char* RawCloudID;
+	ArAnchor_acquireCloudAnchorId(SessionHandle, AnchorHandle, &RawCloudID);
+	FString CloudID(ANSI_TO_TCHAR(RawCloudID));
+	CloudARPin->UpdateCloudState(ToARPinCloudState(NewCloudState), CloudID);
+	ArString_release(RawCloudID);
+} 
+#else // ARCORE_USE_OLD_CLOUD_ANCHOR_ASYNC
 void FGoogleARCoreCloudARPinManager::UpdateCloudARPin(UCloudARPin* CloudARPin, ArSession* SessionHandle)
 {
 	ArFuture* Future = CloudARPin->GetFuture()->AsFuture();
@@ -386,8 +402,96 @@ void FGoogleARCoreCloudARPinManager::UpdateCloudARPin(UCloudARPin* CloudARPin, A
 		check(false);
 	}
 }
+#endif // else ARCORE_USE_OLD_CLOUD_ANCHOR_ASYNC
 #endif // ARCORE_SERVICE_SUPPORTED_PLATFORM
 
+#if ARCORE_USE_OLD_CLOUD_ANCHOR_ASYNC
+UCloudARPin* FGoogleARCoreCloudARPinManager::CreateAndHostCloudARPin(UARPin* PinToHost, int32 LifetimeInDays, EARPinCloudTaskResult& OutTaskResult)
+{
+#if ARCORE_SERVICE_SUPPORTED_PLATFORM
+	if (PinToHost == nullptr)
+	{
+		OutTaskResult = EARPinCloudTaskResult::InvalidPin;
+		return nullptr;
+	}
+
+	float WorldToMeterScale = ARSystem->GetXRTrackingSystem()->GetWorldToMetersScale();
+	ArSession* SessionHandle = GetSessionHandle();
+	ArAnchor* NewAnchorHandle = nullptr;
+#if PLATFORM_ANDROID
+	ArAnchor* AnchorHandle = reinterpret_cast<ArAnchor*> (PinToHost->GetNativeResource());
+	ensure(AnchorHandle != nullptr);
+
+	OutTaskResult = ToCloudTaskResult(ArSession_hostAndAcquireNewCloudAnchor(SessionHandle, AnchorHandle, &NewAnchorHandle));
+#elif PLATFORM_IOS
+	ARKitAnchor* AnchorHandle = nullptr;
+	ArPose *PoseHandle = nullptr;
+	ArPose_create(SessionHandle, nullptr, &PoseHandle);
+	UnrealTransformToARCorePose(PinToHost->GetLocalToTrackingTransform_NoAlignment(), SessionHandle, &PoseHandle, WorldToMeterScale);
+	ARKitAnchor_create(PoseHandle, &AnchorHandle);
+
+	ArPose_destroy(PoseHandle);
+	PoseHandle = nullptr;
+
+	OutTaskResult = ToCloudTaskResult(ArSession_hostAndAcquireNewCloudAnchor(SessionHandle, AnchorHandle, &NewAnchorHandle));
+#endif
+
+	UE_LOG(LogGoogleARCoreServices, Log, TEXT("ArSession_hostAndAcquireNewCloudAnchor returns TaskResult: %u, Anchor: %p"), (int)OutTaskResult, NewAnchorHandle);
+
+	UCloudARPin* NewCloudARPin = nullptr;
+
+	if (OutTaskResult == EARPinCloudTaskResult::Started)
+	{
+		NewCloudARPin = NewObject<UCloudARPin>();
+		FTransform PinTrackingTransform = PinToHost->GetLocalToTrackingTransform_NoAlignment();
+		NewCloudARPin->InitARPin(ARSystem, nullptr, PinTrackingTransform, PinToHost->GetTrackedGeometry(), FName(TEXT("Cloud AR Pin(Hosted)")));
+
+		UpdateCloudARPin(NewCloudARPin, SessionHandle, NewAnchorHandle);
+
+		AllCloudARPins.Add(NewCloudARPin);
+		HandleToCloudPinMap.Add(NewAnchorHandle, NewCloudARPin);
+	}
+
+	return NewCloudARPin;
+#endif
+
+	return nullptr;
+}
+
+UCloudARPin* FGoogleARCoreCloudARPinManager::ResolveAndCreateCloudARPin(FString CloudID, EARPinCloudTaskResult& OutTaskResult)
+{
+#if ARCORE_SERVICE_SUPPORTED_PLATFORM
+	ArSession* SessionHandle = GetSessionHandle();
+	ArAnchor* AnchorHandle = nullptr;
+	ArPose *PoseHandle = nullptr;
+	float WorldToMeterScale = ARSystem->GetXRTrackingSystem()->GetWorldToMetersScale();
+	ArPose_create(SessionHandle, nullptr, &PoseHandle);
+
+	OutTaskResult = ToCloudTaskResult(
+		ArSession_resolveAndAcquireNewCloudAnchor(SessionHandle, TCHAR_TO_ANSI(*CloudID), &AnchorHandle));
+	UE_LOG(LogGoogleARCoreServices, Log, TEXT("ArSession_resolveAndAcquireNewCloudAnchor returns TaskResult: %d, Anchor: %p, CloudID: %s"), (int)OutTaskResult, AnchorHandle, *CloudID);
+	ensure(AnchorHandle != nullptr);
+
+	ArAnchor_getPose(SessionHandle, AnchorHandle, PoseHandle);
+
+	UCloudARPin* NewCloudARPin = NewObject<UCloudARPin>();
+	FTransform PinTrackingTransform = ARCorePoseToUnrealTransform(PoseHandle, SessionHandle, WorldToMeterScale);
+	NewCloudARPin->InitARPin(ARSystem, nullptr, PinTrackingTransform, nullptr, FName(TEXT("Cloud AR Pin(Acquired)")));
+
+	UpdateCloudARPin(NewCloudARPin, SessionHandle, AnchorHandle);
+
+	AllCloudARPins.Add(NewCloudARPin);
+	HandleToCloudPinMap.Add(AnchorHandle, NewCloudARPin);
+
+	ArPose_destroy(PoseHandle);
+
+	return NewCloudARPin;
+#endif
+
+	return nullptr;
+}
+
+#else // ARCORE_USE_OLD_CLOUD_ANCHOR_ASYNC
 UCloudARPin* FGoogleARCoreCloudARPinManager::CreateAndHostCloudARPin(UARPin* PinToHost, int32 LifetimeInDays, EARPinCloudTaskResult& OutTaskResult)
 {
 #if ARCORE_SERVICE_SUPPORTED_PLATFORM
@@ -467,6 +571,7 @@ UCloudARPin* FGoogleARCoreCloudARPinManager::ResolveAndCreateCloudARPin(FString
 	return nullptr;
 #endif
 }
+#endif // else ARCORE_USE_OLD_CLOUD_ANCHOR_ASYNC
 
 void FGoogleARCoreCloudARPinManager::RemoveCloudARPin(UCloudARPin* PinToRemove)
 {
@@ -513,6 +618,7 @@ void FGoogleARCoreCloudARPinManager::UpdateAllCloudARPins()
 		return;
 	}
 
+#if !ARCORE_USE_OLD_CLOUD_ANCHOR_ASYNC
 	// Update any PendingCloudARPins
 	for (UCloudARPin* CloudPin : PendingCloudARPins)
 	{
@@ -520,6 +626,7 @@ void FGoogleARCoreCloudARPinManager::UpdateAllCloudARPins()
 	}
 	// Remove no longer pending CloudARPins
 	PendingCloudARPins.RemoveAllSwap([](const UCloudARPin* CloudPin) { return !CloudPin->IsPending(); });
+#endif
 
 	float WorldToMeterScale = ARSystem->GetXRTrackingSystem()->GetWorldToMetersScale();
 
@@ -562,6 +669,11 @@ void FGoogleARCoreCloudARPinManager::UpdateAllCloudARPins()
 			CloudPin->OnTransformUpdated(AnchorPose);
 		}
 
+#if ARCORE_USE_OLD_CLOUD_ANCHOR_ASYNC
+		// Update Cloud state and ID.
+		UpdateCloudARPin(CloudPin, SessionHandle, AnchorHandle);
+#endif
+
 		ArAnchor_release(AnchorHandle);
 	}
 	ArAnchorList_destroy(UpdatedAnchorListHandle);
diff --git a/Engine/Plugins/Runtime/AR/Google/GoogleARCoreServices/Source/GoogleARCoreServices/Private/GoogleARCoreCloudARPinManager.h b/Engine/Plugins/Runtime/AR/Google/GoogleARCoreServices/Source/GoogleARCoreServices/Private/GoogleARCoreCloudARPinManager.h
--- a/Engine/Plugins/Runtime/AR/Google/GoogleARCoreServices/Source/GoogleARCoreServices/Private/GoogleARCoreCloudARPinManager.h
+++ b/Engine/Plugins/Runtime/AR/Google/GoogleARCoreServices/Source/GoogleARCoreServices/Private/GoogleARCoreCloudARPinManager.h
@@ -74,11 +74,13 @@ protected:
 	// protected properties:
 	TSharedRef<FARSupportInterface, ESPMode::ThreadSafe> ARSystem;
 	TArray<TObjectPtr<UCloudARPin>> AllCloudARPins;
+#if !ARCORE_USE_OLD_CLOUD_ANCHOR_ASYNC
 	TArray<TObjectPtr<UCloudARPin>> PendingCloudARPins; // CloudARPins that have an async action pending.
-
+#endif
 #if	ARCORE_SERVICE_SUPPORTED_PLATFORM
+#if !ARCORE_USE_OLD_CLOUD_ANCHOR_ASYNC
 	virtual void UpdateCloudARPin(UCloudARPin* CloudARPin, ArSession* SessionHandle);
-
+#endif
 	virtual ArSession* GetSessionHandle() = 0;
 
 	virtual ArFrame* GetARFrameHandle() = 0;
diff --git a/Engine/Plugins/Runtime/AR/Google/GoogleARCoreServices/Source/GoogleARCoreServices/Private/GoogleARCoreServicesTypes.cpp b/Engine/Plugins/Runtime/AR/Google/GoogleARCoreServices/Source/GoogleARCoreServices/Private/GoogleARCoreServicesTypes.cpp
--- a/Engine/Plugins/Runtime/AR/Google/GoogleARCoreServices/Source/GoogleARCoreServices/Private/GoogleARCoreServicesTypes.cpp
+++ b/Engine/Plugins/Runtime/AR/Google/GoogleARCoreServices/Source/GoogleARCoreServices/Private/GoogleARCoreServicesTypes.cpp
@@ -3,6 +3,32 @@
 #include "GoogleARCoreServicesTypes.h"
 #include "GoogleARCoreUtils.h"
 
+#if ARCORE_USE_OLD_CLOUD_ANCHOR_ASYNC
+UCloudARPin::UCloudARPin()
+	: UARPin()
+{
+	CloudState = ECloudARPinCloudState::NotHosted;
+	CloudID = FString("");
+}
+
+FString UCloudARPin::GetCloudID()
+{
+	return CloudID;
+}
+
+ECloudARPinCloudState UCloudARPin::GetARPinCloudState()
+{
+	return CloudState;
+}
+
+void UCloudARPin::UpdateCloudState(ECloudARPinCloudState NewCloudState, FString NewCloudID)
+{
+	CloudState = NewCloudState;
+	CloudID = NewCloudID;
+}
+
+#else //ARCORE_USE_OLD_CLOUD_ANCHOR_ASYNC
+
 GoogleARFutureHolderPtr FGoogleARFutureHolder::MakeHostFuture()
 {
 	return MakeShared<FGoogleARFutureHolder>(ECloudAnchorFutureType::Host);
@@ -104,4 +130,4 @@ void UCloudARPin::ReleaseFuture()
 {
 	Future = nullptr;
 }
-
+#endif
diff --git a/Engine/Plugins/Runtime/AR/Google/GoogleARCoreServices/Source/GoogleARCoreServices/Public/GoogleARCoreServicesTypes.h b/Engine/Plugins/Runtime/AR/Google/GoogleARCoreServices/Source/GoogleARCoreServices/Public/GoogleARCoreServicesTypes.h
--- a/Engine/Plugins/Runtime/AR/Google/GoogleARCoreServices/Source/GoogleARCoreServices/Public/GoogleARCoreServicesTypes.h
+++ b/Engine/Plugins/Runtime/AR/Google/GoogleARCoreServices/Source/GoogleARCoreServices/Public/GoogleARCoreServicesTypes.h
@@ -8,9 +8,18 @@
 
 DEFINE_LOG_CATEGORY_STATIC(LogGoogleARCoreServices, Log, All);
 
+
+#if PLATFORM_IOS
+#define ARCORE_USE_OLD_CLOUD_ANCHOR_ASYNC 1
+#else
+#define ARCORE_USE_OLD_CLOUD_ANCHOR_ASYNC 0
+#endif
+
+#if !ARCORE_USE_OLD_CLOUD_ANCHOR_ASYNC
 typedef struct ArFuture_ ArFuture;
 typedef struct ArHostCloudAnchorFuture_ ArHostCloudAnchorFuture;
 typedef struct ArResolveCloudAnchorFuture_ ArResolveCloudAnchorFuture;
+#endif
 
 /// @defgroup GoogleARCoreServices Google ARCore Services
 /// The module for GoogleARCoreServices plugin
@@ -144,6 +153,8 @@ enum class ECloudARPinCloudState : uint8
 	ErrorSDKVersionTooNew
 };
 
+#if ARCORE_USE_OLD_CLOUD_ANCHOR_ASYNC
+#else
 enum class ECloudAnchorFutureType : uint8
 {
 	Unknown,
@@ -177,6 +188,7 @@ private:
 		ArResolveCloudAnchorFuture* ResolveFuture;
 	};
 };
+#endif // !ARCORE_USE_OLD_CLOUD_ANCHOR_ASYNC
 
 /**
  * A CloudARPin will be created when you host an existing ARPin, or resolved a
@@ -207,6 +219,9 @@ public:
 	UFUNCTION(BlueprintPure, Category = "GoogleARCoreServices|CloudARPin")
 	ECloudARPinCloudState GetARPinCloudState();
 
+#if ARCORE_USE_OLD_CLOUD_ANCHOR_ASYNC
+	void UpdateCloudState(ECloudARPinCloudState NewCloudState, FString NewCloudID);
+#else
 	void SetFuture(GoogleARFutureHolderPtr FutureHolder);
 	GoogleARFutureHolderPtr GetFuture() const;
 	bool IsPending() const;
@@ -215,9 +230,13 @@ public:
 	void UpdateCloudState(ECloudARPinCloudState NewCloudState);
 
 	void SetCloudID(FString NewCloudID);
+	#endif
 
 private:
 	ECloudARPinCloudState CloudState;
 	FString CloudID;
+#if ARCORE_USE_OLD_CLOUD_ANCHOR_ASYNC
+#else
 	GoogleARFutureHolderPtr Future;
+#endif
 };
